Added support for the upcoming Iris replication system as part of the new CQTestExperimental plugin

#rnx

diff --git a/Engine/Plugins/Tests/CQTest/Source/CQTestTests/CQTestTests.Build.cs b/Engine/Plugins/Tests/CQTest/Source/CQTestTests/CQTestTests.Build.cs
--- a/Engine/Plugins/Tests/CQTest/Source/CQTestTests/CQTestTests.Build.cs
+++ b/Engine/Plugins/Tests/CQTest/Source/CQTestTests/CQTestTests.Build.cs
@@ -29,5 +29,7 @@ public class CQTestTests : ModuleRules
 					"UnrealEd"
 			});
 		}
+
+		SetupIrisSupport(Target);
 	}
 }
diff --git a/Engine/Plugins/Tests/CQTest/Source/CQTestTests/Private/Components/PIENetworkComponentTests.cpp b/Engine/Plugins/Tests/CQTest/Source/CQTestTests/Private/Components/PIENetworkComponentTests.cpp
--- a/Engine/Plugins/Tests/CQTest/Source/CQTestTests/Private/Components/PIENetworkComponentTests.cpp
+++ b/Engine/Plugins/Tests/CQTest/Source/CQTestTests/Private/Components/PIENetworkComponentTests.cpp
@@ -6,14 +6,17 @@
 #include "CQTestGameMode.h"
 #include "CQTestUnitTestHelper.h"
 #include "TestReplicatedActor.h"
+#include "Tests/AutomationCommon.h"
 
-#include "Kismet/GameplayStatics.h"
-#include "GameFramework/GameModeBase.h"
-#include "Engine/GameInstance.h"
 #include "Engine/World.h"
 
 #if ENABLE_PIE_NETWORK_TEST
 
+/**
+ * The tests below showcase how the `FBasePIENetworkComponent` can be used to setup and simulate a client/server network using the default replication system.
+ * Please enable the CQTestExperimental Engine plugin to see and run tests against the experimental Iris replication system (https://dev.epicgames.com/documentation/en-us/unreal-engine/iris-replication-system-in-unreal-engine).
+*/
+
 NETWORK_TEST_CLASS(StateTest, "TestFramework.CQTest.Network")
 {
 	struct DerivedState : public FBasePIENetworkComponentState
@@ -23,6 +26,13 @@ NETWORK_TEST_CLASS(StateTest, "TestFramework.CQTest.Network")
 
 	int32 SharedNumber = 0;
 	FPIENetworkComponent<DerivedState> Network{ TestRunner, TestCommandBuilder, bInitializing };
+	
+	inline static TSharedPtr<FScopedTestEnvironment> TestEnvironment{ nullptr };
+	BEFORE_ALL()
+	{
+		TestEnvironment = FScopedTestEnvironment::Get();
+		TestEnvironment->SetConsoleVariableValue(TEXT("net.Iris.UseIrisReplication"), TEXT("0"));
+	}
 
 	BEFORE_EACH()
 	{
@@ -32,6 +42,11 @@ NETWORK_TEST_CLASS(StateTest, "TestFramework.CQTest.Network")
 			.WithGameMode(AGameModeBase::StaticClass())
 			.Build(Network);
 	}
+	
+	AFTER_ALL()
+	{
+		TestEnvironment = nullptr;
+	}
 
 	TEST_METHOD(Network_WithMultipleSteps_TriggersStepsInOrder)
 	{
@@ -93,8 +108,15 @@ NETWORK_TEST_CLASS(ReplicationTest, "TestFramework.CQTest.Network")
 	};
 
 	const int32 ExpectedReplicatedValue = 42;
-
 	FPIENetworkComponent<DerivedState> Network{ TestRunner, TestCommandBuilder, bInitializing };
+	
+	inline static TSharedPtr<FScopedTestEnvironment> TestEnvironment{ nullptr };
+	BEFORE_ALL()
+	{
+		TestEnvironment = FScopedTestEnvironment::Get();
+		TestEnvironment->SetConsoleVariableValue(TEXT("net.Iris.UseIrisReplication"), TEXT("0"));
+	}
+
 	BEFORE_EACH() 
 	{
 		FNetworkComponentBuilder<DerivedState>()
@@ -103,6 +125,11 @@ NETWORK_TEST_CLASS(ReplicationTest, "TestFramework.CQTest.Network")
 			.Build(Network);
 	}
 
+	AFTER_ALL()
+	{
+		TestEnvironment = nullptr;
+	}
+
 	TEST_METHOD(SpawnAndReplicateActor_WithReplicatedActor_ProvidesActorToClients)
 	{
 		Network.SpawnAndReplicate<ATestReplicatedActor, &DerivedState::ReplicatedActor>()
@@ -174,6 +201,14 @@ NETWORK_TEST_CLASS(MultipleActorStateReplication, "TestFramework.CQTest.Network"
 	const int32 Actor2ExpectedReplicatedValue = 24;
 
 	FPIENetworkComponent<DerivedState> Network{ TestRunner, TestCommandBuilder, bInitializing };
+
+	inline static TSharedPtr<FScopedTestEnvironment> TestEnvironment{ nullptr };
+	BEFORE_ALL()
+	{
+		TestEnvironment = FScopedTestEnvironment::Get();
+		TestEnvironment->SetConsoleVariableValue(TEXT("net.Iris.UseIrisReplication"), TEXT("0"));
+	}
+
 	BEFORE_EACH()
 	{
 		FNetworkComponentBuilder<DerivedState>()
@@ -181,6 +216,11 @@ NETWORK_TEST_CLASS(MultipleActorStateReplication, "TestFramework.CQTest.Network"
 			.Build(Network);
 	}
 
+	AFTER_ALL()
+	{
+		TestEnvironment = nullptr;
+	}
+
 	auto MakeSetInt(int32 Value) 
 	{
 		return [Value](ATestReplicatedActor& Actor) {
@@ -226,6 +266,14 @@ NETWORK_TEST_CLASS(LateJoinTest, "TestFramework.CQTest.Network")
 	};
 
 	FPIENetworkComponent<DerivedState> Network{ TestRunner, TestCommandBuilder, bInitializing };
+
+	inline static TSharedPtr<FScopedTestEnvironment> TestEnvironment{ nullptr };
+	BEFORE_ALL()
+	{
+		TestEnvironment = FScopedTestEnvironment::Get();
+		TestEnvironment->SetConsoleVariableValue(TEXT("net.Iris.UseIrisReplication"), TEXT("0"));
+	}
+
 	BEFORE_EACH()
 	{
 		FNetworkComponentBuilder<DerivedState>()
@@ -235,6 +283,11 @@ NETWORK_TEST_CLASS(LateJoinTest, "TestFramework.CQTest.Network")
 			.Build(Network);
 	}
 
+	AFTER_ALL()
+	{
+		TestEnvironment = nullptr;
+	}
+
 	TEST_METHOD(ThenClientJoins_AfterStart_AddsClient) 
 	{
 		Network.ThenServer([this](DerivedState& ServerState) {
@@ -269,6 +322,13 @@ NETWORK_TEST_CLASS(GameInstanceTest, "TestFramework.CQTest.Network")
 
 	FPIENetworkComponent<FBasePIENetworkComponentState> Network{ TestRunner, TestCommandBuilder, bInitializing };
 
+	inline static TSharedPtr<FScopedTestEnvironment> TestEnvironment{ nullptr };
+	BEFORE_ALL()
+	{
+		TestEnvironment = FScopedTestEnvironment::Get();
+		TestEnvironment->SetConsoleVariableValue(TEXT("net.Iris.UseIrisReplication"), TEXT("0"));
+	}
+
 	BEFORE_EACH()
 	{
 		FNetworkComponentBuilder()
@@ -277,6 +337,11 @@ NETWORK_TEST_CLASS(GameInstanceTest, "TestFramework.CQTest.Network")
 			.Build(Network);
 	}
 
+	AFTER_ALL()
+	{
+		TestEnvironment = nullptr;
+	}
+
 	TEST_METHOD(NetworkComponent_WithGameInstanceClass_BuildsNetworkWithProvidedGameInstance)
 	{
 		Network.ThenServer([this](FBasePIENetworkComponentState& ServerState) {
@@ -293,6 +358,13 @@ NETWORK_TEST_CLASS(GameModeTest, "TestFramework.CQTest.Network")
 
 	FPIENetworkComponent<> Network{ TestRunner, TestCommandBuilder, bInitializing };
 
+	inline static TSharedPtr<FScopedTestEnvironment> TestEnvironment{ nullptr };
+	BEFORE_ALL()
+	{
+		TestEnvironment = FScopedTestEnvironment::Get();
+		TestEnvironment->SetConsoleVariableValue(TEXT("net.Iris.UseIrisReplication"), TEXT("0"));
+	}
+
 	BEFORE_EACH()
 	{
 		FNetworkComponentBuilder()
@@ -301,6 +373,11 @@ NETWORK_TEST_CLASS(GameModeTest, "TestFramework.CQTest.Network")
 			.Build(Network);
 	}
 
+	AFTER_ALL()
+	{
+		TestEnvironment = nullptr;
+	}
+
 	TEST_METHOD(NetworkComponent_WithGameMode_BuildsNetworkWithProvidedGameMode)
 	{
 		Network.ThenServer([this](FBasePIENetworkComponentState& ServerState) {
@@ -315,11 +392,23 @@ NETWORK_TEST_CLASS(SetupErrorTest, "TestFramework.CQTest.Network")
 {
 	FPIENetworkComponent<> Network{ TestRunner, TestCommandBuilder, bInitializing };
 
+	inline static TSharedPtr<FScopedTestEnvironment> TestEnvironment{ nullptr };
+	BEFORE_ALL()
+	{
+		TestEnvironment = FScopedTestEnvironment::Get();
+		TestEnvironment->SetConsoleVariableValue(TEXT("net.Iris.UseIrisReplication"), TEXT("0"));
+	}
+
 	AFTER_EACH() 
 	{
 		ClearExpectedError(*TestRunner, TEXT("Failed to initialize Network Component"));
 	}
 
+	AFTER_ALL()
+	{
+		TestEnvironment = nullptr;
+	}
+
 	TEST_METHOD(NetworkComponent_WithoutUsingBuilder_AddsErrorAndDoesNotCrash)
 	{
 		Network.ThenServer([this](FBasePIENetworkComponentState&) { Assert.Fail("Unexpected Error"); });
diff --git a/Engine/Plugins/Tests/CQTestExperimental/CQTestExperimental.uplugin b/Engine/Plugins/Tests/CQTestExperimental/CQTestExperimental.uplugin
new file mode 100644
--- /dev/null
+++ b/Engine/Plugins/Tests/CQTestExperimental/CQTestExperimental.uplugin
@@ -0,0 +1,24 @@
+{
+	"FileVersion": 3,
+	"Version": 1,
+	"VersionName": "1.0",
+	"FriendlyName": "Experimental Code Quality Unreal Test Plugin",
+	"Description": "Simplified testing of experimental features for Unreal Engine",
+	"Category": "Testing",
+	"CreatedBy": "Epic Games, Inc.",
+	"CreatedByURL": "https://epicgames.com",
+	"DocsURL": "",
+	"MarketplaceURL": "",
+	"SupportURL": "",
+	"CanContainContent": false,
+	"IsBetaVersion": false,
+	"IsExperimentalVersion": true,
+	"Installed": false,
+	"Modules": [
+		{
+			"Name": "CQTestExperimentalTests",
+			"Type": "Runtime",
+			"LoadingPhase": "Default"
+		}
+	]
+}
\ No newline at end of file
diff --git a/Engine/Plugins/Tests/CQTestExperimental/Source/CQTestExperimentalTests/CQTestExperimentalTests.Build.cs b/Engine/Plugins/Tests/CQTestExperimental/Source/CQTestExperimentalTests/CQTestExperimentalTests.Build.cs
new file mode 100644
--- /dev/null
+++ b/Engine/Plugins/Tests/CQTestExperimental/Source/CQTestExperimentalTests/CQTestExperimentalTests.Build.cs
@@ -0,0 +1,34 @@
+// Copyright Epic Games, Inc. All Rights Reserved.
+
+using UnrealBuildTool;
+
+public class CQTestExperimentalTests : ModuleRules
+{
+	public CQTestExperimentalTests(ReadOnlyTargetRules Target)
+		: base(Target)
+	{
+		PCHUsage = PCHUsageMode.UseExplicitOrSharedPCHs;
+
+		PrivateDependencyModuleNames.AddRange(
+			new string[] {
+					"Core",
+					"CoreUObject",
+					"Engine",
+					"NetCore",
+					"CQTest"
+				 }
+			);
+
+		if (Target.bBuildEditor)
+		{
+			PrivateDependencyModuleNames.AddRange(
+				new string[] {
+					"EngineSettings",
+					"LevelEditor",
+					"UnrealEd"
+			});
+		}
+
+		SetupIrisSupport(Target);
+	}
+}
diff --git a/Engine/Plugins/Tests/CQTestExperimental/Source/CQTestExperimentalTests/CQTestExperimentalTestsModule.cpp b/Engine/Plugins/Tests/CQTestExperimental/Source/CQTestExperimentalTests/CQTestExperimentalTestsModule.cpp
new file mode 100644
--- /dev/null
+++ b/Engine/Plugins/Tests/CQTestExperimental/Source/CQTestExperimentalTests/CQTestExperimentalTestsModule.cpp
@@ -0,0 +1,5 @@
+// Copyright Epic Games, Inc. All Rights Reserved.
+
+#include "Modules/ModuleManager.h"
+
+IMPLEMENT_MODULE( FDefaultModuleImpl, CQTestExperimentalTests )
diff --git a/Engine/Plugins/Tests/CQTestExperimental/Source/CQTestExperimentalTests/Private/Iris/CQTestObjects.cpp b/Engine/Plugins/Tests/CQTestExperimental/Source/CQTestExperimentalTests/Private/Iris/CQTestObjects.cpp
new file mode 100644
--- /dev/null
+++ b/Engine/Plugins/Tests/CQTestExperimental/Source/CQTestExperimentalTests/Private/Iris/CQTestObjects.cpp
@@ -0,0 +1,22 @@
+// Copyright Epic Games, Inc. All Rights Reserved.
+
+#include "CQTestObjects.h"
+
+#include "Net/UnrealNetwork.h"
+
+///////////////////////////////////////////////////////////////////////
+// AIrisTestReplicatedActor
+
+AIrisTestReplicatedActor::AIrisTestReplicatedActor(const FObjectInitializer& ObjectInitializer)
+	: Super(ObjectInitializer)
+	, ReplicatedInt(0)
+{
+	bReplicates = true;
+	bAlwaysRelevant = true;
+}
+
+void AIrisTestReplicatedActor::GetLifetimeReplicatedProps(TArray<FLifetimeProperty>& OutLifetimeProps) const
+{
+	Super::GetLifetimeReplicatedProps(OutLifetimeProps);
+	DOREPLIFETIME(AIrisTestReplicatedActor, ReplicatedInt);
+}
\ No newline at end of file
diff --git a/Engine/Plugins/Tests/CQTestExperimental/Source/CQTestExperimentalTests/Private/Iris/CQTestObjects.h b/Engine/Plugins/Tests/CQTestExperimental/Source/CQTestExperimentalTests/Private/Iris/CQTestObjects.h
new file mode 100644
--- /dev/null
+++ b/Engine/Plugins/Tests/CQTestExperimental/Source/CQTestExperimentalTests/Private/Iris/CQTestObjects.h
@@ -0,0 +1,44 @@
+// Copyright Epic Games, Inc. All Rights Reserved.
+
+#pragma once
+
+#include "Engine/GameInstance.h"
+#include "GameFramework/Actor.h"
+#include "GameFramework/GameModeBase.h"
+
+#include "CQTestObjects.generated.h"
+
+/** Basic GameInstance with a test value. */
+UCLASS()
+class UIrisTestGameInstanceClass : public UGameInstance
+{
+	GENERATED_BODY()
+
+public:
+	int TestValue{ 42 };
+};
+
+/** Basic GameMode with a test value. */
+UCLASS()
+class AIrisTestGameMode : public AGameModeBase
+{
+	GENERATED_BODY()
+
+public:
+	int32 TestValue{ 42 };
+};
+
+/** Basic Actor used for replication tests. */
+UCLASS(NotBlueprintable)
+class AIrisTestReplicatedActor : public AActor
+{
+	GENERATED_BODY()
+
+public:
+	AIrisTestReplicatedActor(const FObjectInitializer& ObjectInitializer = FObjectInitializer::Get());
+
+	void GetLifetimeReplicatedProps(TArray<FLifetimeProperty>& OutLifetimeProps) const override;
+
+	UPROPERTY(Replicated)
+	int32 ReplicatedInt;
+};
\ No newline at end of file
diff --git a/Engine/Plugins/Tests/CQTestExperimental/Source/CQTestExperimentalTests/Private/Iris/PIENetworkComponentTests.cpp b/Engine/Plugins/Tests/CQTestExperimental/Source/CQTestExperimentalTests/Private/Iris/PIENetworkComponentTests.cpp
new file mode 100644
--- /dev/null
+++ b/Engine/Plugins/Tests/CQTestExperimental/Source/CQTestExperimentalTests/Private/Iris/PIENetworkComponentTests.cpp
@@ -0,0 +1,427 @@
+// Copyright Epic Games, Inc. All Rights Reserved.
+
+#include "CQTest.h"
+#include "Components/PIENetworkComponent.h"
+#include "CQTestObjects.h"
+#include "Tests/AutomationCommon.h"
+
+#include "Engine/World.h"
+
+#if ENABLE_PIE_NETWORK_TEST && UE_WITH_IRIS
+
+/**
+ * The tests below showcase how the `FBasePIENetworkComponent` can be used to setup and simulate a client/server network using the experimental iris replication system (https://dev.epicgames.com/documentation/en-us/unreal-engine/iris-replication-system-in-unreal-engine).
+ * See the CQTest Engine plugin for more information on tests using the default replication system.
+*/
+
+NETWORK_TEST_CLASS(IrisStateTest, "TestFramework.CQTest.Experimental.Network")
+{
+	struct DerivedState : public FBasePIENetworkComponentState
+	{
+		int32 IndependentNumber = 0;
+	};
+
+	int32 SharedNumber = 0;
+	FPIENetworkComponent<DerivedState> Network{ TestRunner, TestCommandBuilder, bInitializing };
+	
+	inline static TSharedPtr<FScopedTestEnvironment> TestEnvironment{ nullptr };
+	BEFORE_ALL()
+	{
+		TestEnvironment = FScopedTestEnvironment::Get();
+		TestEnvironment->SetConsoleVariableValue(TEXT("net.Iris.UseIrisReplication"), TEXT("1"));
+	}
+
+	BEFORE_EACH()
+	{
+		FNetworkComponentBuilder<DerivedState>()
+			.WithClients(3)
+			.WithGameInstanceClass(UGameInstance::StaticClass())
+			.WithGameMode(AGameModeBase::StaticClass())
+			.Build(Network);
+	}
+	
+	AFTER_ALL()
+	{
+		TestEnvironment = nullptr;
+	}
+
+	TEST_METHOD(Network_WithMultipleSteps_TriggersStepsInOrder)
+	{
+		Network
+			.ThenServer([this](DerivedState&) { ASSERT_THAT(AreEqual(0, SharedNumber++)); })
+			.ThenClient(0, [this](DerivedState&) { ASSERT_THAT(AreEqual(1, SharedNumber++)); })
+			.ThenServer([this](DerivedState&) { ASSERT_THAT(AreEqual(2, SharedNumber++)); })
+			.Then([this]() { ASSERT_THAT(AreEqual(3, SharedNumber++)); });
+	}
+
+	TEST_METHOD(Network_WithServerCommands_RetainsStateBetweenCalls)
+	{
+		Network
+			.ThenServer([this](DerivedState& State) { ASSERT_THAT(AreEqual(0, State.IndependentNumber++)); })
+			.ThenServer([this](DerivedState& State) { ASSERT_THAT(AreEqual(1, State.IndependentNumber++)); });
+	}
+
+	TEST_METHOD(Network_WithClientCommands_RetainsStateBetweenCalls)
+	{
+		Network
+			.ThenClient(0, [this](DerivedState& State) { ASSERT_THAT(AreEqual(0, State.IndependentNumber++)); })
+			.ThenClient(0, [this](DerivedState& State) { ASSERT_THAT(AreEqual(1, State.IndependentNumber++)); });
+	}
+
+	TEST_METHOD(Network_WithClientAndServerCommands_DoNotShareState)
+	{
+		Network
+			.ThenServer([this](DerivedState& State) { State.IndependentNumber++; })
+			.ThenClient(0, [this](DerivedState& State) { ASSERT_THAT(AreEqual(0, State.IndependentNumber)); });
+	}
+
+	TEST_METHOD(Network_WithMultipleClients_DoNotShareState)
+	{
+		Network
+			.ThenClients([this](DerivedState& State) { State.IndependentNumber = State.ClientIndex; })
+			.ThenClients([this](DerivedState& State) { ASSERT_THAT(AreEqual(State.ClientIndex, State.IndependentNumber)); });
+	}
+
+	TEST_METHOD(Network_WithTickingServerCommand_TicksUntilDone)
+	{
+		Network
+			.UntilServer([this](DerivedState& State) { return ++State.IndependentNumber > 4; })
+			.ThenServer([this](DerivedState& State) { ASSERT_THAT(AreEqual(State.IndependentNumber, 5)); });
+	}
+
+	TEST_METHOD(Network_WithTickingClientCommands_TicksEachCommand)
+	{
+		Network
+			.UntilClients([this](DerivedState& State) { SharedNumber++;  return ++State.IndependentNumber > 4; })
+			.Then([this]() { ASSERT_THAT(AreEqual(15, SharedNumber)); });
+	}
+};
+
+NETWORK_TEST_CLASS(IrisReplicationTest, "TestFramework.CQTest.Experimental.Network")
+{
+	struct DerivedState : public FBasePIENetworkComponentState
+	{
+		AIrisTestReplicatedActor* ReplicatedActor = nullptr;
+	};
+
+	const int32 ExpectedReplicatedValue = 42;
+	FPIENetworkComponent<DerivedState> Network{ TestRunner, TestCommandBuilder, bInitializing };
+	
+	inline static TSharedPtr<FScopedTestEnvironment> TestEnvironment{ nullptr };
+	BEFORE_ALL()
+	{
+		TestEnvironment = FScopedTestEnvironment::Get();
+		TestEnvironment->SetConsoleVariableValue(TEXT("net.Iris.UseIrisReplication"), TEXT("1"));
+	}
+
+	BEFORE_EACH() 
+	{
+		FNetworkComponentBuilder<DerivedState>()
+			.WithGameInstanceClass(UGameInstance::StaticClass())
+			.WithGameMode(AGameModeBase::StaticClass())
+			.Build(Network);
+	}
+
+	AFTER_ALL()
+	{
+		TestEnvironment = nullptr;
+	}
+
+	TEST_METHOD(SpawnAndReplicateActor_WithReplicatedActor_ProvidesActorToClients)
+	{
+		Network.SpawnAndReplicate<AIrisTestReplicatedActor, &DerivedState::ReplicatedActor>()
+			.ThenServer([this](DerivedState& ServerState) {
+				ASSERT_THAT(IsNotNull(ServerState.ReplicatedActor));
+			})
+			.ThenClients([this](DerivedState& ClientState) {
+				ASSERT_THAT(IsNotNull(ClientState.ReplicatedActor));
+			});
+	}
+
+	TEST_METHOD(SpawnAndReplicateActor_ThenUpdateProperty_UpdatesPropertyOnClients)
+	{
+		Network.SpawnAndReplicate<AIrisTestReplicatedActor, &DerivedState::ReplicatedActor>()
+			.ThenServer(TEXT("Server Set Value"), [this](DerivedState& ServerState) {
+				ServerState.ReplicatedActor->ReplicatedInt = ExpectedReplicatedValue;
+			})
+			.UntilClients(TEXT("Clients Check Value"), [this](DerivedState& ClientState) {
+				return ClientState.ReplicatedActor->ReplicatedInt == ExpectedReplicatedValue;
+			});
+	}
+
+	TEST_METHOD(SpawnAndReplicateActor_WithSpawnParameters_PassesParametersToSpawnedObject)
+	{
+		FActorSpawnParameters SpawnParameters;
+		SpawnParameters.SpawnCollisionHandlingOverride = ESpawnActorCollisionHandlingMethod::AlwaysSpawn;
+		Network.SpawnAndReplicate<AIrisTestReplicatedActor, &DerivedState::ReplicatedActor>(SpawnParameters)
+			.UntilClients([SpawnParameters](DerivedState& ClientState) {
+				return SpawnParameters.SpawnCollisionHandlingOverride == ClientState.ReplicatedActor->SpawnCollisionHandlingMethod;
+			});
+	}
+
+	TEST_METHOD(SpawnAndReplicateActor_WithBeforeReplicates_InvokesBeforeReplicate) {
+		auto BeforeReplicate = [this](AIrisTestReplicatedActor& Actor) {
+			Actor.ReplicatedInt = ExpectedReplicatedValue;
+		};
+
+		Network.SpawnAndReplicate<AIrisTestReplicatedActor, &DerivedState::ReplicatedActor>(BeforeReplicate)
+			.UntilClients([this](DerivedState& ClientState) {
+				return ClientState.ReplicatedActor->ReplicatedInt == ExpectedReplicatedValue;
+			});
+	}
+
+	TEST_METHOD(SpawnAndReplicateActor_WithSpawnParametersAndBeforeReplicate_UsesBoth)
+	{
+		auto BeforeReplicate = [this](AIrisTestReplicatedActor& Actor) {
+			Actor.ReplicatedInt = ExpectedReplicatedValue;
+		};
+
+		FActorSpawnParameters SpawnParameters;
+		SpawnParameters.SpawnCollisionHandlingOverride = ESpawnActorCollisionHandlingMethod::AlwaysSpawn;
+		Network.SpawnAndReplicate<AIrisTestReplicatedActor, &DerivedState::ReplicatedActor>(SpawnParameters, BeforeReplicate)
+			.UntilClients([this, SpawnParameters](DerivedState& ClientState) {
+			AIrisTestReplicatedActor& actor = *ClientState.ReplicatedActor;
+				return SpawnParameters.SpawnCollisionHandlingOverride == actor.SpawnCollisionHandlingMethod && actor.ReplicatedInt == ExpectedReplicatedValue;
+			});
+	}
+};
+
+NETWORK_TEST_CLASS(IrisMultipleActorStateReplication, "TestFramework.CQTest.Experimental.Network")
+{
+	struct DerivedState : public FBasePIENetworkComponentState
+	{
+		AIrisTestReplicatedActor* ReplicatedActor1 = nullptr;
+		AIrisTestReplicatedActor* ReplicatedActor2 = nullptr;
+	};
+
+	const int32 Actor1ExpectedReplicatedValue = 42;
+	const int32 Actor2ExpectedReplicatedValue = 24;
+
+	FPIENetworkComponent<DerivedState> Network{ TestRunner, TestCommandBuilder, bInitializing };
+
+	inline static TSharedPtr<FScopedTestEnvironment> TestEnvironment{ nullptr };
+	BEFORE_ALL()
+	{
+		TestEnvironment = FScopedTestEnvironment::Get();
+		TestEnvironment->SetConsoleVariableValue(TEXT("net.Iris.UseIrisReplication"), TEXT("1"));
+	}
+
+	BEFORE_EACH()
+	{
+		FNetworkComponentBuilder<DerivedState>()
+			.WithClients(1)
+			.Build(Network);
+	}
+
+	AFTER_ALL()
+	{
+		TestEnvironment = nullptr;
+	}
+
+	auto MakeSetInt(int32 Value) 
+	{
+		return [Value](AIrisTestReplicatedActor& Actor) {
+			Actor.ReplicatedInt = Value;
+		};
+	}
+
+	TEST_METHOD(SpawnAndReplicateActor_WithMultipleActors_ReplicatesBoth)
+	{
+		Network.SpawnAndReplicate<AIrisTestReplicatedActor, &DerivedState::ReplicatedActor1>(MakeSetInt(Actor1ExpectedReplicatedValue))
+			.SpawnAndReplicate<AIrisTestReplicatedActor, &DerivedState::ReplicatedActor2>(MakeSetInt(Actor2ExpectedReplicatedValue))
+			.ThenServer([this](DerivedState& ServerState) {
+				ASSERT_THAT(AreEqual(Actor1ExpectedReplicatedValue, ServerState.ReplicatedActor1->ReplicatedInt));
+				ASSERT_THAT(AreEqual(Actor2ExpectedReplicatedValue, ServerState.ReplicatedActor2->ReplicatedInt));
+			})
+			.ThenClients([this](DerivedState& ClientState) {
+				ASSERT_THAT(AreEqual(Actor1ExpectedReplicatedValue, ClientState.ReplicatedActor1->ReplicatedInt));
+				ASSERT_THAT(AreEqual(Actor2ExpectedReplicatedValue, ClientState.ReplicatedActor2->ReplicatedInt));
+			});
+	}
+
+	TEST_METHOD(SpawnAndReplicate_WithLateJoin_ReplicatesBoth)
+	{
+		Network.SpawnAndReplicate<AIrisTestReplicatedActor, &DerivedState::ReplicatedActor1>(MakeSetInt(42))
+			.SpawnAndReplicate<AIrisTestReplicatedActor, &DerivedState::ReplicatedActor2>(MakeSetInt(24))
+			.ThenServer([this](DerivedState& ServerState) {
+				ASSERT_THAT(AreEqual(Actor1ExpectedReplicatedValue, ServerState.ReplicatedActor1->ReplicatedInt));
+				ASSERT_THAT(AreEqual(Actor2ExpectedReplicatedValue, ServerState.ReplicatedActor2->ReplicatedInt));
+			})
+			.ThenClientJoins()
+			.ThenClients([this](DerivedState& ClientState) {
+				ASSERT_THAT(AreEqual(Actor1ExpectedReplicatedValue, ClientState.ReplicatedActor1->ReplicatedInt));
+				ASSERT_THAT(AreEqual(Actor2ExpectedReplicatedValue, ClientState.ReplicatedActor2->ReplicatedInt));
+			});
+	}
+};
+
+NETWORK_TEST_CLASS(IrisLateJoinTest, "TestFramework.CQTest.Experimental.Network")
+{
+	struct DerivedState : public FBasePIENetworkComponentState
+	{
+		AIrisTestReplicatedActor* ReplicatedActor = nullptr;
+	};
+
+	FPIENetworkComponent<DerivedState> Network{ TestRunner, TestCommandBuilder, bInitializing };
+
+	inline static TSharedPtr<FScopedTestEnvironment> TestEnvironment{ nullptr };
+	BEFORE_ALL()
+	{
+		TestEnvironment = FScopedTestEnvironment::Get();
+		TestEnvironment->SetConsoleVariableValue(TEXT("net.Iris.UseIrisReplication"), TEXT("1"));
+	}
+
+	BEFORE_EACH()
+	{
+		FNetworkComponentBuilder<DerivedState>()
+			.WithGameInstanceClass(UGameInstance::StaticClass())
+			.WithGameMode(AGameModeBase::StaticClass())
+			.WithClients(1)
+			.Build(Network);
+	}
+
+	AFTER_ALL()
+	{
+		TestEnvironment = nullptr;
+	}
+
+	TEST_METHOD(ThenClientJoins_AfterStart_AddsClient) 
+	{
+		Network.ThenServer([this](DerivedState& ServerState) {
+			ASSERT_THAT(AreEqual(ServerState.ClientCount, 1));
+		})
+		.ThenClientJoins()
+		.ThenServer([this](DerivedState& ServerState) {
+			ASSERT_THAT(AreEqual(ServerState.ClientCount, 2));
+			ASSERT_THAT(AreEqual(ServerState.ClientConnections.Num(), 2));
+		});
+	}
+
+	TEST_METHOD(ThenClientJoins_AfterStart_ReplicatesState) 
+	{
+		Network.SpawnAndReplicate<AIrisTestReplicatedActor, &DerivedState::ReplicatedActor>()
+			.ThenServer([this](DerivedState& ServerState) {
+				ASSERT_THAT(IsNotNull(ServerState.ReplicatedActor));
+			})
+			.ThenClient(0, [this](DerivedState& ClientState) {
+				ASSERT_THAT(IsNotNull(ClientState.ReplicatedActor));
+			})
+			.ThenClientJoins()
+			.ThenClient(1, [this](DerivedState& ClientState) {
+				ASSERT_THAT(IsNotNull(ClientState.ReplicatedActor));
+			});
+	}
+};
+
+NETWORK_TEST_CLASS(IrisGameInstanceTest, "TestFramework.CQTest.Experimental.Network")
+{
+	const int32 ExpectedReplicatedValue = 42;
+
+	FPIENetworkComponent<FBasePIENetworkComponentState> Network{ TestRunner, TestCommandBuilder, bInitializing };
+
+	inline static TSharedPtr<FScopedTestEnvironment> TestEnvironment{ nullptr };
+	BEFORE_ALL()
+	{
+		TestEnvironment = FScopedTestEnvironment::Get();
+		TestEnvironment->SetConsoleVariableValue(TEXT("net.Iris.UseIrisReplication"), TEXT("1"));
+	}
+
+	BEFORE_EACH()
+	{
+		FNetworkComponentBuilder()
+			.WithGameInstanceClass(FSoftClassPath(UIrisTestGameInstanceClass::StaticClass()))
+			.WithGameMode(AGameModeBase::StaticClass())
+			.Build(Network);
+	}
+
+	AFTER_ALL()
+	{
+		TestEnvironment = nullptr;
+	}
+
+	TEST_METHOD(NetworkComponent_WithGameInstanceClass_BuildsNetworkWithProvidedGameInstance)
+	{
+		Network.ThenServer([this](FBasePIENetworkComponentState& ServerState) {
+			UIrisTestGameInstanceClass* GameInstance = Cast<UIrisTestGameInstanceClass>(ServerState.World->GetGameInstance());
+			ASSERT_THAT(IsNotNull(GameInstance));
+			ASSERT_THAT(AreEqual(ExpectedReplicatedValue, GameInstance->TestValue));
+		});
+	}
+};
+
+NETWORK_TEST_CLASS(IrisGameModeTest, "TestFramework.CQTest.Experimental.Network")
+{
+	const int32 ExpectedReplicatedValue = 42;
+
+	FPIENetworkComponent<> Network{ TestRunner, TestCommandBuilder, bInitializing };
+
+	inline static TSharedPtr<FScopedTestEnvironment> TestEnvironment{ nullptr };
+	BEFORE_ALL()
+	{
+		TestEnvironment = FScopedTestEnvironment::Get();
+		TestEnvironment->SetConsoleVariableValue(TEXT("net.Iris.UseIrisReplication"), TEXT("1"));
+	}
+
+	BEFORE_EACH()
+	{
+		FNetworkComponentBuilder()
+			.WithGameInstanceClass(UGameInstance::StaticClass())
+			.WithGameMode(AIrisTestGameMode::StaticClass())
+			.Build(Network);
+	}
+
+	AFTER_ALL()
+	{
+		TestEnvironment = nullptr;
+	}
+
+	TEST_METHOD(NetworkComponent_WithGameMode_BuildsNetworkWithProvidedGameMode)
+	{
+		Network.ThenServer([this](FBasePIENetworkComponentState& ServerState) {
+			AIrisTestGameMode* GameMode = Cast<AIrisTestGameMode>(ServerState.World->GetAuthGameMode());
+			ASSERT_THAT(IsNotNull(GameMode));
+			ASSERT_THAT(AreEqual(ExpectedReplicatedValue, GameMode->TestValue));
+		});
+	}
+};
+
+NETWORK_TEST_CLASS(IrisSetupErrorTest, "TestFramework.CQTest.Experimental.Network")
+{
+	FPIENetworkComponent<> Network{ TestRunner, TestCommandBuilder, bInitializing };
+
+	inline static TSharedPtr<FScopedTestEnvironment> TestEnvironment{ nullptr };
+	BEFORE_ALL()
+	{
+		TestEnvironment = FScopedTestEnvironment::Get();
+		TestEnvironment->SetConsoleVariableValue(TEXT("net.Iris.UseIrisReplication"), TEXT("1"));
+	}
+
+	AFTER_EACH() 
+	{
+		const FString ExpectedError = TEXT("Failed to initialize Network Component");
+		
+		FAutomationTestExecutionInfo TestInfo;
+		TestRunner->GetExecutionInfo(TestInfo);
+
+		TestInfo.RemoveAllEvents([&ExpectedError](FAutomationEvent& event) {
+			return event.Message.Contains(ExpectedError);
+		});
+
+		if (TestInfo.GetErrorTotal() == 0)
+		{
+			TestRunner->ClearExecutionInfo();
+		}
+	}
+
+	AFTER_ALL()
+	{
+		TestEnvironment = nullptr;
+	}
+
+	TEST_METHOD(NetworkComponent_WithoutUsingBuilder_AddsErrorAndDoesNotCrash)
+	{
+		Network.ThenServer([this](FBasePIENetworkComponentState&) { Assert.Fail("Unexpected Error"); });
+	}
+};
+
+#endif // ENABLE_PIE_NETWORK_TEST && UE_WITH_IRIS
\ No newline at end of file
diff --git a/Engine/Source/Developer/CQTest/CQTest.Build.cs b/Engine/Source/Developer/CQTest/CQTest.Build.cs
--- a/Engine/Source/Developer/CQTest/CQTest.Build.cs
+++ b/Engine/Source/Developer/CQTest/CQTest.Build.cs
@@ -14,6 +14,7 @@ public class CQTest : ModuleRules
 					"Core",
 					"CoreUObject",
 					"Engine",
+					"NetCore",
 					"Slate"
 				 }
 			);
@@ -27,5 +28,7 @@ public class CQTest : ModuleRules
 					"UnrealEd"
 			});
 		}
+
+		SetupIrisSupport(Target);
 	}
 }
diff --git a/Engine/Source/Developer/CQTest/Public/Components/PIENetworkComponent.h b/Engine/Source/Developer/CQTest/Public/Components/PIENetworkComponent.h
--- a/Engine/Source/Developer/CQTest/Public/Components/PIENetworkComponent.h
+++ b/Engine/Source/Developer/CQTest/Public/Components/PIENetworkComponent.h
@@ -14,8 +14,14 @@
 #include "Editor.h"
 #include "Engine/NetDriver.h"
 #include "Engine/PackageMapClient.h"
+#include "Net/NetIDVariant.h"
 #include "UnrealEdGlobals.h"
 
+#if UE_WITH_IRIS
+#include "Iris/ReplicationSystem/ReplicationSystem.h"
+#include "Iris/ReplicationSystem/ObjectReplicationBridge.h"
+#endif
+
 /*
 //Example boiler plate
 
@@ -85,7 +91,7 @@ struct FBasePIENetworkComponentState
 	bool bIsDedicatedServer = true;
 
 	/** Used to track spawned and replicated actors across client and server PIE sessions. */
-	TSet<FNetworkGUID> LocallySpawnedActors{};
+	TSet<UE::Net::FNetIDVariant> LocallySpawnedActors{};
 };
 
 /**
@@ -250,7 +256,7 @@ protected:
 	FPacketSimulationSettings* PacketSimulationSettings = nullptr;
 	TSubclassOf<AGameModeBase> GameMode = TSubclassOf<AGameModeBase>(nullptr);
 	FPIENetworkTestStateRestorer StateRestorer;
-	TMap<FNetworkGUID, int64> SpawnedActors{};
+	TMap<UE::Net::FNetIDVariant, int64> SpawnedActors{};
 };
 
 template <typename NetworkDataType>
diff --git a/Engine/Source/Developer/CQTest/Public/Impl/PIENetworkComponent.inl b/Engine/Source/Developer/CQTest/Public/Impl/PIENetworkComponent.inl
--- a/Engine/Source/Developer/CQTest/Public/Impl/PIENetworkComponent.inl
+++ b/Engine/Source/Developer/CQTest/Public/Impl/PIENetworkComponent.inl
@@ -215,9 +215,25 @@ inline FPIENetworkComponent<NetworkDataType>& FPIENetworkComponent<NetworkDataTy
 			State.*ResultStorage = *ServerActor;
 		}
 	})
-	.UntilServer(TEXT("Waiting for NetGUID"), [this, ServerActor](NetworkDataType& State) {
-		FNetworkGUID NetGUID = State.World->GetNetDriver()->GetNetGuidCache()->GetNetGUID(*ServerActor);
-		if (!NetGUID.IsValid())
+	.UntilServer(TEXT("Waiting for Net ID"), [this, ServerActor](NetworkDataType& State) {
+		UE::Net::FNetIDVariant NetIDVariant;
+		if (State.World->GetNetDriver()->IsUsingIrisReplication())
+		{
+#if UE_WITH_IRIS
+			UReplicationSystem* ReplicationSystem = State.World->GetNetDriver()->GetReplicationSystem();
+			UObjectReplicationBridge* ReplicationBridge = Cast<UObjectReplicationBridge>(ReplicationSystem->GetReplicationBridge());
+			checkf(IsValid(ReplicationBridge), TEXT("Unable to create a ReplicationBridge."));
+			NetIDVariant = UE::Net::FNetIDVariant(ReplicationBridge->GetReplicatedRefHandle(*ServerActor));
+#else
+			checkf(false, TEXT("NetDriver is set to use Iris Replication, but the Engine is not configured with Iris support."));
+#endif // UE_WITH_IRIS
+		}
+		else
+		{
+			NetIDVariant = UE::Net::FNetIDVariant(State.World->GetNetDriver()->GetNetGuidCache()->GetNetGUID(*ServerActor));
+		}
+
+		if (!NetIDVariant.IsValid())
 		{
 			return false;
 		}
@@ -230,8 +246,8 @@ inline FPIENetworkComponent<NetworkDataType>& FPIENetworkComponent<NetworkDataTy
 		NetworkDataType TempDataType {};
 		int64 StorageOffset = reinterpret_cast<int64>(&(TempDataType.*ResultStorage)) - reinterpret_cast<int64>(&TempDataType);
 
-		SpawnedActors.Add(NetGUID, StorageOffset);
-		State.LocallySpawnedActors.Add(NetGUID);
+		SpawnedActors.Add(NetIDVariant, StorageOffset);
+		State.LocallySpawnedActors.Add(NetIDVariant);
 		return true;
 	}, TimeoutValue);
 	
@@ -241,14 +257,30 @@ inline FPIENetworkComponent<NetworkDataType>& FPIENetworkComponent<NetworkDataTy
 template <typename NetworkDataType>
 inline bool FPIENetworkComponent<NetworkDataType>::ReplicateToClients(NetworkDataType& ClientState)
 {
-	for(const auto& [NetGUID, StorageOffset] : SpawnedActors)
+	for(const auto& [NetIDVariant, StorageOffset] : SpawnedActors)
 	{
-		if(ClientState.LocallySpawnedActors.Contains(NetGUID))
+		if(ClientState.LocallySpawnedActors.Contains(NetIDVariant))
 		{
 			continue;
 		}
 
-		AActor* ClientActor = Cast<AActor>(ClientState.World->GetNetDriver()->GetNetGuidCache()->GetObjectFromNetGUID(NetGUID, true));
+		AActor* ClientActor = nullptr;
+		if (ClientState.World->GetNetDriver()->IsUsingIrisReplication())
+		{
+#if UE_WITH_IRIS
+			UReplicationSystem* ReplicationSystem = ClientState.World->GetNetDriver()->GetReplicationSystem();
+			UObjectReplicationBridge* ReplicationBridge = Cast<UObjectReplicationBridge>(ReplicationSystem->GetReplicationBridge());
+			checkf(IsValid(ReplicationBridge), TEXT("Unable to create a ReplicationBridge."));
+			ClientActor = Cast<AActor>(ReplicationBridge->GetReplicatedObject(NetIDVariant.GetVariant().template Get<UE::Net::FNetRefHandle>()));
+#else
+			checkf(false, TEXT("NetDriver is set to use Iris Replication, but the Engine is not configured with Iris support."));
+#endif // UE_WITH_IRIS
+		}
+		else
+		{
+			ClientActor = Cast<AActor>(ClientState.World->GetNetDriver()->GetNetGuidCache()->GetObjectFromNetGUID(NetIDVariant.GetVariant().template Get<FNetworkGUID>(), true));
+		}
+
 		if (ClientActor != nullptr)
 		{
 			// The other half of the offset implementation.
@@ -256,7 +288,7 @@ inline bool FPIENetworkComponent<NetworkDataType>::ReplicateToClients(NetworkDat
 			// Then interpret that as a pointer to a pointer to an AActor which we can assign
 			AActor** Storage = reinterpret_cast<AActor**>(reinterpret_cast<char*>(&ClientState) + StorageOffset);
 			*Storage = ClientActor;
-			ClientState.LocallySpawnedActors.Add(NetGUID);
+			ClientState.LocallySpawnedActors.Add(NetIDVariant);
 		}
 	}
 
