Mark CQ testing modules as DeveloperTool

#rnx

#Synced-CL 50554652

diff --git a/Engine/Plugins/Tests/AudioCQTests/AudioCQTests.uplugin b/Engine/Plugins/Tests/AudioCQTests/AudioCQTests.uplugin
--- a/Engine/Plugins/Tests/AudioCQTests/AudioCQTests.uplugin
+++ b/Engine/Plugins/Tests/AudioCQTests/AudioCQTests.uplugin
@@ -16,7 +16,7 @@
   "Modules": [
     {
       "Name": "AudioCQTests",
-      "Type": "Runtime",
+      "Type": "DeveloperTool",
       "LoadingPhase": "Default"
     }
   ]
diff --git a/Engine/Plugins/Tests/CQTestExperimental/CQTestExperimental.uplugin b/Engine/Plugins/Tests/CQTestExperimental/CQTestExperimental.uplugin
--- a/Engine/Plugins/Tests/CQTestExperimental/CQTestExperimental.uplugin
+++ b/Engine/Plugins/Tests/CQTestExperimental/CQTestExperimental.uplugin
@@ -17,7 +17,7 @@
 	"Modules": [
 		{
 			"Name": "CQTestExperimentalTests",
-			"Type": "Runtime",
+			"Type": "DeveloperTool",
 			"LoadingPhase": "Default"
 		}
 	]
